Rename  IniKeyBlacklist to IniKeyDenylist and IniSectionBlacklist to IniSectionDenyList

diff --git a/Engine/Config/BaseGame.ini b/Engine/Config/BaseGame.ini
--- a/Engine/Config/BaseGame.ini
+++ b/Engine/Config/BaseGame.ini
@@ -130,25 +130,25 @@ WriteBackMetadataToAssetRegistry=Disabled
 +EarlyDownloaderPakFileFiles=...\*.uproject
 +EarlyDownloaderPakFileFiles=...\global_sf*.metalmap
 
-+IniKeyBlacklist=KeyStorePassword
-+IniKeyBlacklist=KeyPassword
-+IniKeyBlacklist=rsa.privateexp
-+IniKeyBlacklist=rsa.modulus
-+IniKeyBlacklist=rsa.publicexp
-+IniKeyBlacklist=aes.key
-+IniKeyBlacklist=SigningPublicExponent
-+IniKeyBlacklist=SigningModulus
-+IniKeyBlacklist=SigningPrivateExponent
-+IniKeyBlacklist=EncryptionKey
-+IniKeyBlacklist=DevCenterUsername
-+IniKeyBlacklist=DevCenterPassword
-+IniKeyBlacklist=IOSTeamID
-+IniKeyBlacklist=SigningCertificate
-+IniKeyBlacklist=MobileProvision
-+IniKeyBlacklist=IniKeyBlacklist
-+IniKeyBlacklist=IniSectionBlacklist
-
-+IniSectionBlacklist=HordeStorageServers
++IniKeyDenylist=KeyStorePassword
++IniKeyDenylist=KeyPassword
++IniKeyDenylist=rsa.privateexp
++IniKeyDenylist=rsa.modulus
++IniKeyDenylist=rsa.publicexp
++IniKeyDenylist=aes.key
++IniKeyDenylist=SigningPublicExponent
++IniKeyDenylist=SigningModulus
++IniKeyDenylist=SigningPrivateExponent
++IniKeyDenylist=EncryptionKey
++IniKeyDenylist=DevCenterUsername
++IniKeyDenylist=DevCenterPassword
++IniKeyDenylist=IOSTeamID
++IniKeyDenylist=SigningCertificate
++IniKeyDenylist=MobileProvision
++IniKeyDenylist=IniKeyDenylist
++IniKeyDenylist=IniSectionDenylist
+
++IniSectionDenylist=HordeStorageServers
 
 [/Script/Engine.HUD]
 DebugDisplay=AI
diff --git a/Engine/Config/ConfigRedirects.ini b/Engine/Config/ConfigRedirects.ini
--- a/Engine/Config/ConfigRedirects.ini
+++ b/Engine/Config/ConfigRedirects.ini
@@ -25,3 +25,6 @@ BlacklistConfigFiles=DisallowedConfigFiles
 WhitelistDirectories=AllowedDirectories
 BlacklistLocalizationTargets=DisallowedLocalizationTargets
 
+[/Script/UnrealEd.ProjectPackagingSettings]
+IniKeyBlacklist=IniKeyDenylist
+IniSectionBlacklist=IniSectionDenylist
diff --git a/Engine/Source/Developer/DeveloperToolSettings/Classes/Settings/ProjectPackagingSettings.h b/Engine/Source/Developer/DeveloperToolSettings/Classes/Settings/ProjectPackagingSettings.h
--- a/Engine/Source/Developer/DeveloperToolSettings/Classes/Settings/ProjectPackagingSettings.h
+++ b/Engine/Source/Developer/DeveloperToolSettings/Classes/Settings/ProjectPackagingSettings.h
@@ -543,14 +543,22 @@ public:
 	UPROPERTY(config, EditAnywhere, Category = Packaging, AdvancedDisplay)
 	TArray<FString> CompressedChunkWildcard;
 
-	/** List of ini file keys to strip when packaging */
+	UE_DEPRECATED(5.1, "This property is no longer supported. Use IniKeyDenylist.")
 	UPROPERTY(config, EditAnywhere, Category = Packaging)
 	TArray<FString> IniKeyBlacklist;
 
-	/** List of ini file sections to strip when packaging */
+	/** List of ini file keys to strip when packaging */
+	UPROPERTY(config, EditAnywhere, Category = Packaging)
+	TArray<FString> IniKeyDenylist;
+
+	UE_DEPRECATED(5.1, "This property is no longer supported. Use IniSectionDenylist.")
 	UPROPERTY(config, EditAnywhere, Category = Packaging)
 	TArray<FString> IniSectionBlacklist;
 
+	/** List of ini file sections to strip when packaging */
+	UPROPERTY(config, EditAnywhere, Category = Packaging)
+	TArray<FString> IniSectionDenylist;
+
 	/**
 	 * List of specific files to include with GenerateEarlyDownloaderPakFile
 	 */
diff --git a/Engine/Source/Developer/PakFileUtilities/Private/PakFileUtilities.cpp b/Engine/Source/Developer/PakFileUtilities/Private/PakFileUtilities.cpp
--- a/Engine/Source/Developer/PakFileUtilities/Private/PakFileUtilities.cpp
+++ b/Engine/Source/Developer/PakFileUtilities/Private/PakFileUtilities.cpp
@@ -5050,8 +5050,8 @@ bool MakeBinaryConfig(const TCHAR* CmdLine)
 
 	TArray<FString> KeyDenyListStrings;
 	TArray<FString> SectionsDenyList;
-	GConfig->GetArray(TEXT("/Script/UnrealEd.ProjectPackagingSettings"), TEXT("IniKeyBlacklist"), KeyDenyListStrings, GGameIni);
-	GConfig->GetArray(TEXT("/Script/UnrealEd.ProjectPackagingSettings"), TEXT("IniSectionBlacklist"), SectionsDenyList, GGameIni);
+	GConfig->GetArray(TEXT("/Script/UnrealEd.ProjectPackagingSettings"), TEXT("IniKeyDenylist"), KeyDenyListStrings, GGameIni);
+	GConfig->GetArray(TEXT("/Script/UnrealEd.ProjectPackagingSettings"), TEXT("IniSectionDenylist"), SectionsDenyList, GGameIni);
 	TArray<FName> KeysDenyList;
 	for (const FString& Key : KeyDenyListStrings)
 	{
@@ -5083,7 +5083,7 @@ bool MakeBinaryConfig(const TCHAR* CmdLine)
 
 	// check the deny list removed itself
 	KeyDenyListStrings.Empty();
-	Config.GetArray(TEXT("/Script/UnrealEd.ProjectPackagingSettings"), TEXT("IniKeyBlacklist"), KeyDenyListStrings, GGameIni);
+	Config.GetArray(TEXT("/Script/UnrealEd.ProjectPackagingSettings"), TEXT("IniKeyDenylist"), KeyDenyListStrings, GGameIni);
 	check(KeyDenyListStrings.Num() == 0);
 
 	// allow delegates to modify the config data with some tagged binary data
diff --git a/Engine/Source/Editor/UnrealEd/Private/Commandlets/MakeBinaryConfigCommandlet.cpp b/Engine/Source/Editor/UnrealEd/Private/Commandlets/MakeBinaryConfigCommandlet.cpp
--- a/Engine/Source/Editor/UnrealEd/Private/Commandlets/MakeBinaryConfigCommandlet.cpp
+++ b/Engine/Source/Editor/UnrealEd/Private/Commandlets/MakeBinaryConfigCommandlet.cpp
@@ -61,8 +61,8 @@ int32 UMakeBinaryConfigCommandlet::Main(const FString& Params)
 
 	TArray<FString> KeyDenyListStrings;
 	TArray<FString> SectionsDenyList;
-	GConfig->GetArray(TEXT("/Script/UnrealEd.ProjectPackagingSettings"), TEXT("IniKeyBlacklist"), KeyDenyListStrings, GGameIni);
-	GConfig->GetArray(TEXT("/Script/UnrealEd.ProjectPackagingSettings"), TEXT("IniSectionBlacklist"), SectionsDenyList, GGameIni);
+	GConfig->GetArray(TEXT("/Script/UnrealEd.ProjectPackagingSettings"), TEXT("IniKeyDenylist"), KeyDenyListStrings, GGameIni);
+	GConfig->GetArray(TEXT("/Script/UnrealEd.ProjectPackagingSettings"), TEXT("IniSectionDenylist"), SectionsDenyList, GGameIni);
 	TArray<FName> KeysDenyList;
 	for (FString Key : KeyDenyListStrings)
 	{
@@ -94,7 +94,7 @@ int32 UMakeBinaryConfigCommandlet::Main(const FString& Params)
 
 	// check the deny list removed itself
 	KeyDenyListStrings.Empty();
-	Config.GetArray(TEXT("/Script/UnrealEd.ProjectPackagingSettings"), TEXT("IniKeyBlacklist"), KeyDenyListStrings, GGameIni);
+	Config.GetArray(TEXT("/Script/UnrealEd.ProjectPackagingSettings"), TEXT("IniKeyDenylist"), KeyDenyListStrings, GGameIni);
 	check(KeyDenyListStrings.Num() == 0);
 
 	// allow delegates to modify the config data with some tagged binary data
diff --git a/Engine/Source/Programs/AutomationTool/AutomationUtils/DeploymentContext.cs b/Engine/Source/Programs/AutomationTool/AutomationUtils/DeploymentContext.cs
--- a/Engine/Source/Programs/AutomationTool/AutomationUtils/DeploymentContext.cs
+++ b/Engine/Source/Programs/AutomationTool/AutomationUtils/DeploymentContext.cs
@@ -598,8 +598,8 @@ public class DeploymentContext //: ProjectParams
 		String PackagingIniPath = "/Script/UnrealEd.ProjectPackagingSettings";
 
 		// Read the config deny lists
-		GameConfig.GetArray(PackagingIniPath, "IniKeyBlacklist", out IniKeyDenyList);
-		GameConfig.GetArray(PackagingIniPath, "IniSectionBlacklist", out IniSectionDenyList);
+		GameConfig.GetArray(PackagingIniPath, "IniKeyDenylist", out IniKeyDenyList);
+		GameConfig.GetArray(PackagingIniPath, "IniSectionDenylist", out IniSectionDenyList);
 
 		// TODO: Drive these lists from a config file
 		IniSuffixAllowList = new List<string>
diff --git a/Templates/TP_AEC_CollabBP/Config/DefaultGame.ini b/Templates/TP_AEC_CollabBP/Config/DefaultGame.ini
--- a/Templates/TP_AEC_CollabBP/Config/DefaultGame.ini
+++ b/Templates/TP_AEC_CollabBP/Config/DefaultGame.ini
@@ -50,30 +50,30 @@ bCookAll=False
 bCookMapsOnly=False
 bSkipEditorContent=False
 bSkipMovies=False
--IniKeyBlacklist=KeyStorePassword
--IniKeyBlacklist=KeyPassword
--IniKeyBlacklist=rsa.privateexp
--IniKeyBlacklist=rsa.modulus
--IniKeyBlacklist=rsa.publicexp
--IniKeyBlacklist=aes.key
--IniKeyBlacklist=SigningPublicExponent
--IniKeyBlacklist=SigningModulus
--IniKeyBlacklist=SigningPrivateExponent
--IniKeyBlacklist=EncryptionKey
--IniKeyBlacklist=IniKeyBlacklist
--IniKeyBlacklist=IniSectionBlacklist
-+IniKeyBlacklist=KeyStorePassword
-+IniKeyBlacklist=KeyPassword
-+IniKeyBlacklist=rsa.privateexp
-+IniKeyBlacklist=rsa.modulus
-+IniKeyBlacklist=rsa.publicexp
-+IniKeyBlacklist=aes.key
-+IniKeyBlacklist=SigningPublicExponent
-+IniKeyBlacklist=SigningModulus
-+IniKeyBlacklist=SigningPrivateExponent
-+IniKeyBlacklist=EncryptionKey
-+IniKeyBlacklist=IniKeyBlacklist
-+IniKeyBlacklist=IniSectionBlacklist
+-IniKeyDenylist=KeyStorePassword
+-IniKeyDenylist=KeyPassword
+-IniKeyDenylist=rsa.privateexp
+-IniKeyDenylist=rsa.modulus
+-IniKeyDenylist=rsa.publicexp
+-IniKeyDenylist=aes.key
+-IniKeyDenylist=SigningPublicExponent
+-IniKeyDenylist=SigningModulus
+-IniKeyDenylist=SigningPrivateExponent
+-IniKeyDenylist=EncryptionKey
+-IniKeyDenylist=IniKeyDenylist
+-IniKeyDenylist=IniSectionDenylist
++IniKeyDenylist=KeyStorePassword
++IniKeyDenylist=KeyPassword
++IniKeyDenylist=rsa.privateexp
++IniKeyDenylist=rsa.modulus
++IniKeyDenylist=rsa.publicexp
++IniKeyDenylist=aes.key
++IniKeyDenylist=SigningPublicExponent
++IniKeyDenylist=SigningModulus
++IniKeyDenylist=SigningPrivateExponent
++IniKeyDenylist=EncryptionKey
++IniKeyDenylist=IniKeyDenylist
++IniKeyDenylist=IniSectionDenylist
 +MapsToCook=(FilePath="/Game/ArchVis/SampleScene/SampleLevel")
 +MapsToCook=(FilePath="/Game/CollaborativeViewer/Maps/LoginLevel")
 +DirectoriesToAlwaysCook=(Path="/DatasmithRuntime/Materials")
diff --git a/Templates/TP_MFG_CollabBP/Config/DefaultGame.ini b/Templates/TP_MFG_CollabBP/Config/DefaultGame.ini
--- a/Templates/TP_MFG_CollabBP/Config/DefaultGame.ini
+++ b/Templates/TP_MFG_CollabBP/Config/DefaultGame.ini
@@ -50,30 +50,30 @@ bCookAll=False
 bCookMapsOnly=False
 bSkipEditorContent=False
 bSkipMovies=False
--IniKeyBlacklist=KeyStorePassword
--IniKeyBlacklist=KeyPassword
--IniKeyBlacklist=rsa.privateexp
--IniKeyBlacklist=rsa.modulus
--IniKeyBlacklist=rsa.publicexp
--IniKeyBlacklist=aes.key
--IniKeyBlacklist=SigningPublicExponent
--IniKeyBlacklist=SigningModulus
--IniKeyBlacklist=SigningPrivateExponent
--IniKeyBlacklist=EncryptionKey
--IniKeyBlacklist=IniKeyBlacklist
--IniKeyBlacklist=IniSectionBlacklist
-+IniKeyBlacklist=KeyStorePassword
-+IniKeyBlacklist=KeyPassword
-+IniKeyBlacklist=rsa.privateexp
-+IniKeyBlacklist=rsa.modulus
-+IniKeyBlacklist=rsa.publicexp
-+IniKeyBlacklist=aes.key
-+IniKeyBlacklist=SigningPublicExponent
-+IniKeyBlacklist=SigningModulus
-+IniKeyBlacklist=SigningPrivateExponent
-+IniKeyBlacklist=EncryptionKey
-+IniKeyBlacklist=IniKeyBlacklist
-+IniKeyBlacklist=IniSectionBlacklist
+-IniKeyDenylist=KeyStorePassword
+-IniKeyDenylist=KeyPassword
+-IniKeyDenylist=rsa.privateexp
+-IniKeyDenylist=rsa.modulus
+-IniKeyDenylist=rsa.publicexp
+-IniKeyDenylist=aes.key
+-IniKeyDenylist=SigningPublicExponent
+-IniKeyDenylist=SigningModulus
+-IniKeyDenylist=SigningPrivateExponent
+-IniKeyDenylist=EncryptionKey
+-IniKeyDenylist=IniKeyDenylist
+-IniKeyDenylist=IniSectionDenylist
++IniKeyDenylist=KeyStorePassword
++IniKeyDenylist=KeyPassword
++IniKeyDenylist=rsa.privateexp
++IniKeyDenylist=rsa.modulus
++IniKeyDenylist=rsa.publicexp
++IniKeyDenylist=aes.key
++IniKeyDenylist=SigningPublicExponent
++IniKeyDenylist=SigningModulus
++IniKeyDenylist=SigningPrivateExponent
++IniKeyDenylist=EncryptionKey
++IniKeyDenylist=IniKeyDenylist
++IniKeyDenylist=IniSectionDenylist
 +MapsToCook=(FilePath="/Game/ConceptCar/SampleScene/SampleLevel")
 +DirectoriesToAlwaysCook=(Path="/DatasmithRuntime/Materials")
 
