Collaborative Viewer | Updated to use Enhanced Input. Blueprints have been cleaned, orphaned nodes and variables removed, typos fixed and comments added. OpenXR Enabled in Automotive version of template. Legacy SteamVR Binding config files and other legacy related config lines removed.

Runtime error and interaction bug fixed by Marc.

diff --git a/Templates/TP_AEC_CollabBP/Config/DefaultEngine.ini b/Templates/TP_AEC_CollabBP/Config/DefaultEngine.ini
--- a/Templates/TP_AEC_CollabBP/Config/DefaultEngine.ini
+++ b/Templates/TP_AEC_CollabBP/Config/DefaultEngine.ini
@@ -81,6 +81,7 @@ GameInstanceClass=/Game/CollaborativeViewer/Blueprints/GameInstance/BP_Collabora
 
 [/Script/Engine.RendererSettings]
 r.CustomDepth=3
+r.Mobile.UseHWsRGBEncoding=True
 
 [/Script/Engine.PhysicsSettings]
 DefaultGravityZ=-980.000000
@@ -139,7 +140,7 @@ bAutoEnabled=False
 MasterSubmix=/Engine/EngineSounds/Submixes/MasterSubmixDefault.MasterSubmixDefault
 VoiPSampleRate=Normal24000Hz
 
-[Voice] 
+[Voice]
 bEnabled=true
 
 [OnlineSubsystem]
@@ -155,7 +156,113 @@ MaxClientRate=200000
 MaxInternetClientRate=200000
 
 [/Script/AndroidRuntimeSettings.AndroidRuntimeSettings]
+PackageName=com.epicgames.[PROJECT]
+StoreVersion=1
+StoreVersionOffsetArm64=0
+StoreVersionOffsetX8664=0
+ApplicationDisplayName=
+VersionDisplayName=1.0
+MinSDKVersion=26
+TargetSDKVersion=28
+InstallLocation=InternalOnly
+bEnableLint=False
+bPackageDataInsideApk=False
+bCreateAllPlatformsInstall=False
+bDisableVerifyOBBOnStartUp=False
+bForceSmallOBBFiles=False
+bAllowLargeOBBFiles=False
+bAllowPatchOBBFile=False
+bAllowOverflowOBBFiles=False
+bUseExternalFilesDir=False
 bPublicLogFiles=True
+Orientation=SensorLandscape
+MaxAspectRatio=2.100000
+bUseDisplayCutout=False
+bRestoreNotificationsOnReboot=False
+bFullScreen=True
+bEnableNewKeyboard=True
+DepthBufferPreference=Default
+bValidateTextureFormats=True
+bForceCompressNativeLibs=False
+bEnableAdvancedBinaryCompression=False
+bEnableBundle=False
+bEnableUniversalAPK=True
+bBundleABISplit=True
+bBundleLanguageSplit=True
+bBundleDensitySplit=True
+ExtraApplicationSettings=
+ExtraActivitySettings=
+bAndroidVoiceEnabled=False
+bEnableMulticastSupport=False
++PackageForOculusMobile=Quest
++PackageForOculusMobile=Quest2
+bRemoveOSIG=False
+KeyStore=UE.keystore
+KeyAlias=DistroKey
+KeyStorePassword=codesigner
+KeyPassword=
+bBuildForArm64=True
+bBuildForX8664=False
+bBuildForES31=True
+bSupportsVulkan=True
+bSupportsVulkanSM5=False
+DebugVulkanLayerDirectory=(Path="")
+bAndroidOpenGLSupportsBackbufferSampling=False
+bDetectVulkanByDefault=True
+bBuildWithHiddenSymbolVisibility=True
+bDisableStackProtector=False
+bDisableLibCppSharedDependencyValidation=False
+bSaveSymbols=False
+bStripShaderReflection=True
+bEnableGooglePlaySupport=False
+bUseGetAccounts=False
+GamesAppID=
+bEnableSnapshots=False
+bSupportAdMob=True
+AdMobAdUnitID=
+GooglePlayLicenseKey=
+GCMClientSenderID=
+bShowLaunchImage=True
+bAllowIMU=False
+bAllowControllers=True
+bBlockAndroidKeysOnControllers=False
+bControllersBlockDeviceFeedback=False
+AndroidAudio=Default
+AudioSampleRate=44100
+AudioCallbackBufferFrameSize=1024
+AudioNumBuffersToEnqueue=4
+AudioMaxChannels=0
+AudioNumSourceWorkers=0
+SpatializationPlugin=
+SourceDataOverridePlugin=
+ReverbPlugin=
+OcclusionPlugin=
+CompressionOverrides=(bOverrideCompressionTimes=False,DurationThreshold=5.000000,MaxNumRandomBranches=0,SoundCueQualityIndex=0)
+CacheSizeKB=65536
+MaxChunkSizeOverrideKB=0
+bResampleForDevice=False
+SoundCueCookQualityIndex=-1
+MaxSampleRate=48000.000000
+HighSampleRate=32000.000000
+MedSampleRate=24000.000000
+LowSampleRate=12000.000000
+MinSampleRate=8000.000000
+CompressionQualityModifier=1.000000
+AutoStreamingThreshold=0.000000
+AndroidGraphicsDebugger=None
+MaliGraphicsDebuggerPath=(Path="")
+bEnableMaliPerfCounters=False
+bMultiTargetFormat_ETC2=True
+bMultiTargetFormat_DXT=True
+bMultiTargetFormat_ASTC=True
+TextureFormatPriority_ETC2=0.200000
+TextureFormatPriority_DXT=0.600000
+TextureFormatPriority_ASTC=0.900000
+SDKAPILevelOverride=
+NDKAPILevelOverride=
+BuildToolsOverride=
+bStreamLandscapeMeshLODs=False
+bEnableDomStorage=False
 
 [/Script/WindowsTargetPlatform.WindowsTargetSettings]
 DefaultGraphicsRHI=DefaultGraphicsRHI_DX12
@@ -163,3 +270,4 @@ DefaultGraphicsRHI=DefaultGraphicsRHI_DX12
 +D3D12TargetedShaderFormats=PCD3D_SM6
 -D3D11TargetedShaderFormats=PCD3D_SM5
 +D3D11TargetedShaderFormats=PCD3D_SM5
+
diff --git a/Templates/TP_AEC_CollabBP/Config/DefaultInput.ini b/Templates/TP_AEC_CollabBP/Config/DefaultInput.ini
--- a/Templates/TP_AEC_CollabBP/Config/DefaultInput.ini
+++ b/Templates/TP_AEC_CollabBP/Config/DefaultInput.ini
@@ -137,6 +137,12 @@ bUseMouseForTouch=False
 bEnableMouseSmoothing=True
 bEnableFOVScaling=True
 bCaptureMouseOnLaunch=True
+bEnableLegacyInputScales=True
+bEnableMotionControls=True
+bFilterInputByPlatformUser=False
+bEnableInputDeviceSubsystem=True
+bShouldFlushPressedKeysOnViewportFocusLost=True
+bEnableDynamicComponentInputBinding=True
 bAlwaysShowTouchInterface=False
 bShowConsoleOnFourFingerTap=True
 bEnableGestureRecognizer=True
@@ -145,95 +151,12 @@ DefaultViewportMouseCaptureMode=CapturePermanently_IncludingInitialMouseDown
 DefaultViewportMouseLockMode=LockOnCapture
 FOVScale=0.011110
 DoubleClickTime=0.200000
-+ActionMappings=(ActionName="Bookmark_1",bShift=False,bCtrl=False,bAlt=False,bCmd=False,Key=One)
-+ActionMappings=(ActionName="Bookmark_2",bShift=False,bCtrl=False,bAlt=False,bCmd=False,Key=Two)
-+ActionMappings=(ActionName="Bookmark_3",bShift=False,bCtrl=False,bAlt=False,bCmd=False,Key=Three)
-+ActionMappings=(ActionName="Bookmark_4",bShift=False,bCtrl=False,bAlt=False,bCmd=False,Key=Four)
-+ActionMappings=(ActionName="Bookmark_5",bShift=False,bCtrl=False,bAlt=False,bCmd=False,Key=Five)
-+ActionMappings=(ActionName="Bookmark_6",bShift=False,bCtrl=False,bAlt=False,bCmd=False,Key=Six)
-+ActionMappings=(ActionName="Bookmark_7",bShift=False,bCtrl=False,bAlt=False,bCmd=False,Key=Seven)
-+ActionMappings=(ActionName="Bookmark_8",bShift=False,bCtrl=False,bAlt=False,bCmd=False,Key=Eight)
-+ActionMappings=(ActionName="Bookmark_9",bShift=False,bCtrl=False,bAlt=False,bCmd=False,Key=Nine)
-+ActionMappings=(ActionName="FlyMode",bShift=False,bCtrl=False,bAlt=False,bCmd=False,Key=U)
-+ActionMappings=(ActionName="LockToAxis",bShift=False,bCtrl=False,bAlt=False,bCmd=False,Key=LeftShift)
-+ActionMappings=(ActionName="OpenContextualMenu",bShift=False,bCtrl=False,bAlt=False,bCmd=False,Key=SpaceBar)
-+ActionMappings=(ActionName="OrbitAroundObject",bShift=False,bCtrl=False,bAlt=False,bCmd=False,Key=RightMouseButton)
-+ActionMappings=(ActionName="OrbitMode",bShift=False,bCtrl=False,bAlt=False,bCmd=False,Key=O)
-+ActionMappings=(ActionName="OrbitZoomIn",bShift=False,bCtrl=False,bAlt=False,bCmd=False,Key=MouseScrollUp)
-+ActionMappings=(ActionName="OrbitZoomOut",bShift=False,bCtrl=False,bAlt=False,bCmd=False,Key=MouseScrollDown)
-+ActionMappings=(ActionName="PushAndPull",bShift=False,bCtrl=False,bAlt=False,bCmd=False,Key=MouseWheelAxis)
-+ActionMappings=(ActionName="Quit",bShift=False,bCtrl=True,bAlt=False,bCmd=False,Key=Q)
-+ActionMappings=(ActionName="RotateView",bShift=False,bCtrl=False,bAlt=False,bCmd=False,Key=RightMouseButton)
-+ActionMappings=(ActionName="SetPointOfinterest",bShift=False,bCtrl=False,bAlt=False,bCmd=False,Key=MiddleMouseButton)
-+ActionMappings=(ActionName="ShootLaser",bShift=False,bCtrl=False,bAlt=False,bCmd=False,Key=LeftMouseButton)
-+ActionMappings=(ActionName="ShowContextualMenu_VR_L",bShift=False,bCtrl=False,bAlt=False,bCmd=False,Key=OculusTouch_Left_X_Click)
-+ActionMappings=(ActionName="ShowContextualMenu_VR_L",bShift=False,bCtrl=False,bAlt=False,bCmd=False,Key=ValveIndex_Left_A_Click)
-+ActionMappings=(ActionName="ShowContextualMenu_VR_L",bShift=False,bCtrl=False,bAlt=False,bCmd=False,Key=Vive_Left_Menu_Click)
-+ActionMappings=(ActionName="ShowContextualMenu_VR_R",bShift=False,bCtrl=False,bAlt=False,bCmd=False,Key=OculusTouch_Right_A_Click)
-+ActionMappings=(ActionName="ShowContextualMenu_VR_R",bShift=False,bCtrl=False,bAlt=False,bCmd=False,Key=ValveIndex_Right_A_Click)
-+ActionMappings=(ActionName="ShowContextualMenu_VR_R",bShift=False,bCtrl=False,bAlt=False,bCmd=False,Key=Vive_Right_Menu_Click)
-+ActionMappings=(ActionName="VR_Teleport_L",bShift=False,bCtrl=False,bAlt=False,bCmd=False,Key=OculusTouch_Left_Y_Click)
-+ActionMappings=(ActionName="VR_Teleport_L",bShift=False,bCtrl=False,bAlt=False,bCmd=False,Key=ValveIndex_Left_B_Click)
-+ActionMappings=(ActionName="VR_Teleport_L",bShift=False,bCtrl=False,bAlt=False,bCmd=False,Key=Vive_Left_Grip_Click)
-+ActionMappings=(ActionName="VR_Teleport_R",bShift=False,bCtrl=False,bAlt=False,bCmd=False,Key=OculusTouch_Right_B_Click)
-+ActionMappings=(ActionName="VR_Teleport_R",bShift=False,bCtrl=False,bAlt=False,bCmd=False,Key=ValveIndex_Right_B_Click)
-+ActionMappings=(ActionName="VR_Teleport_R",bShift=False,bCtrl=False,bAlt=False,bCmd=False,Key=Vive_Right_Grip_Click)
-+ActionMappings=(ActionName="VR_Trigger_L",bShift=False,bCtrl=False,bAlt=False,bCmd=False,Key=OculusTouch_Left_Trigger_Click)
-+ActionMappings=(ActionName="VR_Trigger_L",bShift=False,bCtrl=False,bAlt=False,bCmd=False,Key=ValveIndex_Left_Trigger_Click)
-+ActionMappings=(ActionName="VR_Trigger_L",bShift=False,bCtrl=False,bAlt=False,bCmd=False,Key=Vive_Left_Trigger_Click)
-+ActionMappings=(ActionName="VR_Trigger_R",bShift=False,bCtrl=False,bAlt=False,bCmd=False,Key=OculusTouch_Right_Trigger_Click)
-+ActionMappings=(ActionName="VR_Trigger_R",bShift=False,bCtrl=False,bAlt=False,bCmd=False,Key=ValveIndex_Right_Trigger_Click)
-+ActionMappings=(ActionName="VR_Trigger_R",bShift=False,bCtrl=False,bAlt=False,bCmd=False,Key=Vive_Right_Trigger_Click)
-+ActionMappings=(ActionName="VRMode",bShift=False,bCtrl=False,bAlt=False,bCmd=False,Key=P)
-+ActionMappings=(ActionName="WalkMode",bShift=False,bCtrl=False,bAlt=False,bCmd=False,Key=I)
-+ActionMappings=(ActionName="NoCommand",bShift=False,bCtrl=False,bAlt=False,bCmd=False,Key=Escape)
-+ActionMappings=(ActionName="TurnRight",bShift=False,bCtrl=False,bAlt=False,bCmd=False,Key=D)
-+ActionMappings=(ActionName="TurnLeft",bShift=False,bCtrl=False,bAlt=False,bCmd=False,Key=A)
-+ActionMappings=(ActionName="TurnRight",bShift=False,bCtrl=False,bAlt=False,bCmd=False,Key=OculusTouch_Right_Thumbstick_Right)
-+ActionMappings=(ActionName="TurnLeft",bShift=False,bCtrl=False,bAlt=False,bCmd=False,Key=OculusTouch_Right_Thumbstick_Left)
-+ActionMappings=(ActionName="VR_Teleport_L",bShift=False,bCtrl=False,bAlt=False,bCmd=False,Key=MixedReality_Left_Grip_Click)
-+ActionMappings=(ActionName="VR_Teleport_R",bShift=False,bCtrl=False,bAlt=False,bCmd=False,Key=MixedReality_Right_Grip_Click)
-+ActionMappings=(ActionName="VR_Trigger_R",bShift=False,bCtrl=False,bAlt=False,bCmd=False,Key=MixedReality_Right_Trigger_Click)
-+ActionMappings=(ActionName="VR_Trigger_L",bShift=False,bCtrl=False,bAlt=False,bCmd=False,Key=MixedReality_Left_Trigger_Click)
-+ActionMappings=(ActionName="ShowContextualMenu_VR_R",bShift=False,bCtrl=False,bAlt=False,bCmd=False,Key=MixedReality_Right_Menu_Click)
-+ActionMappings=(ActionName="ShowContextualMenu_VR_L",bShift=False,bCtrl=False,bAlt=False,bCmd=False,Key=MixedReality_Left_Menu_Click)
-+AxisMappings=(AxisName="FlyUp",Scale=1.000000,Key=E)
-+AxisMappings=(AxisName="FlyUp",Scale=-1.000000,Key=Q)
-+AxisMappings=(AxisName="LookRight",Scale=1.000000,Key=Gamepad_RightX)
-+AxisMappings=(AxisName="LookRight",Scale=1.000000,Key=MouseX)
-+AxisMappings=(AxisName="LookUp",Scale=1.000000,Key=Gamepad_RightY)
-+AxisMappings=(AxisName="LookUp",Scale=1.000000,Key=MouseY)
-+AxisMappings=(AxisName="MoveForward",Scale=-1.000000,Key=Down)
-+AxisMappings=(AxisName="MoveForward",Scale=1.000000,Key=Gamepad_LeftY)
-+AxisMappings=(AxisName="MoveForward",Scale=-1.000000,Key=S)
-+AxisMappings=(AxisName="MoveForward",Scale=1.000000,Key=Up)
-+AxisMappings=(AxisName="MoveForward",Scale=1.000000,Key=W)
-+AxisMappings=(AxisName="MoveRight",Scale=-1.000000,Key=A)
-+AxisMappings=(AxisName="MoveRight",Scale=1.000000,Key=D)
-+AxisMappings=(AxisName="MoveRight",Scale=1.000000,Key=Gamepad_LeftX)
-+AxisMappings=(AxisName="pinch",Scale=1.000000,Key=Gesture_Pinch)
-+AxisMappings=(AxisName="Thumbstick_L_Y",Scale=-1.000000,Key=OculusTouch_Left_Thumbstick_Y)
-+AxisMappings=(AxisName="Thumbstick_L_Y",Scale=-1.000000,Key=ValveIndex_Left_Thumbstick_Y)
-+AxisMappings=(AxisName="Thumbstick_R_Y",Scale=-1.000000,Key=OculusTouch_Right_Thumbstick_Y)
-+AxisMappings=(AxisName="Thumbstick_R_Y",Scale=-1.000000,Key=ValveIndex_Right_Thumbstick_Y)
-+AxisMappings=(AxisName="ThumbStick_Y_Axis",Scale=1.000000,Key=MouseWheelAxis)
-+AxisMappings=(AxisName="ThumbStick_Y_Axis",Scale=1.000000,Key=OculusTouch_Left_Thumbstick_Y)
-+AxisMappings=(AxisName="ThumbStick_Y_Axis",Scale=1.000000,Key=OculusTouch_Right_Thumbstick_Y)
-+AxisMappings=(AxisName="WalkForward",Scale=-1.000000,Key=Down)
-+AxisMappings=(AxisName="WalkForward",Scale=1.000000,Key=Gamepad_LeftY)
-+AxisMappings=(AxisName="WalkForward",Scale=-1.000000,Key=S)
-+AxisMappings=(AxisName="WalkForward",Scale=1.000000,Key=Up)
-+AxisMappings=(AxisName="WalkForward",Scale=1.000000,Key=W)
-+AxisMappings=(AxisName="WalkRight",Scale=-1.000000,Key=A)
-+AxisMappings=(AxisName="WalkRight",Scale=1.000000,Key=D)
-+AxisMappings=(AxisName="WalkRight",Scale=1.000000,Key=Gamepad_LeftX)
-+AxisMappings=(AxisName="WalkRight",Scale=-1.000000,Key=Left)
-+AxisMappings=(AxisName="WalkRight",Scale=1.000000,Key=Right)
-+AxisMappings=(AxisName="Thumbstick_L_Y",Scale=-1.000000,Key=Vive_Left_Trackpad_Y)
-+AxisMappings=(AxisName="Thumbstick_R_Y",Scale=-1.000000,Key=Vive_Right_Trackpad_Y)
-DefaultPlayerInputClass=/Script/Engine.PlayerInput
-DefaultInputComponentClass=/Script/Engine.InputComponent
+DefaultPlayerInputClass=/Script/EnhancedInput.EnhancedPlayerInput
+DefaultInputComponentClass=/Script/EnhancedInput.EnhancedInputComponent
 DefaultTouchInterface=/Game/CollaborativeViewer/Blueprints/Pawn/Touch/TouchInterface.TouchInterface
 -ConsoleKeys=Tilde
 +ConsoleKeys=Tilde
 
+[/Script/OpenXRInput.OpenXRInputSettings]
+MappableInputConfig=/Game/CollaborativeViewer/Input/PMI_CollabTemplate.PMI_CollabTemplate
+
diff --git a/Templates/TP_AEC_CollabBP/Config/SteamVRBindings/gamepad.json b/Templates/TP_AEC_CollabBP/Config/SteamVRBindings/gamepad.json
deleted file mode 100644
diff --git a/Templates/TP_AEC_CollabBP/Config/SteamVRBindings/holographic_controller.json b/Templates/TP_AEC_CollabBP/Config/SteamVRBindings/holographic_controller.json
deleted file mode 100644
diff --git a/Templates/TP_AEC_CollabBP/Config/SteamVRBindings/hpmotioncontroller.json b/Templates/TP_AEC_CollabBP/Config/SteamVRBindings/hpmotioncontroller.json
deleted file mode 100644
diff --git a/Templates/TP_AEC_CollabBP/Config/SteamVRBindings/indexhmd.json b/Templates/TP_AEC_CollabBP/Config/SteamVRBindings/indexhmd.json
deleted file mode 100644
diff --git a/Templates/TP_AEC_CollabBP/Config/SteamVRBindings/knuckles.json b/Templates/TP_AEC_CollabBP/Config/SteamVRBindings/knuckles.json
deleted file mode 100644
diff --git a/Templates/TP_AEC_CollabBP/Config/SteamVRBindings/oculus_touch.json b/Templates/TP_AEC_CollabBP/Config/SteamVRBindings/oculus_touch.json
deleted file mode 100644
diff --git a/Templates/TP_AEC_CollabBP/Config/SteamVRBindings/rift.json b/Templates/TP_AEC_CollabBP/Config/SteamVRBindings/rift.json
deleted file mode 100644
diff --git a/Templates/TP_AEC_CollabBP/Config/SteamVRBindings/steamvr_manifest.json b/Templates/TP_AEC_CollabBP/Config/SteamVRBindings/steamvr_manifest.json
deleted file mode 100644
diff --git a/Templates/TP_AEC_CollabBP/Config/SteamVRBindings/vive.json b/Templates/TP_AEC_CollabBP/Config/SteamVRBindings/vive.json
deleted file mode 100644
diff --git a/Templates/TP_AEC_CollabBP/Config/SteamVRBindings/vive_controller.json b/Templates/TP_AEC_CollabBP/Config/SteamVRBindings/vive_controller.json
deleted file mode 100644
diff --git a/Templates/TP_AEC_CollabBP/Config/SteamVRBindings/vive_cosmos_controller.json b/Templates/TP_AEC_CollabBP/Config/SteamVRBindings/vive_cosmos_controller.json
deleted file mode 100644
diff --git a/Templates/TP_AEC_CollabBP/Config/SteamVRBindings/vive_pro.json b/Templates/TP_AEC_CollabBP/Config/SteamVRBindings/vive_pro.json
deleted file mode 100644
diff --git a/Templates/TP_AEC_CollabBP/Config/SteamVRBindings/vive_tracker.json b/Templates/TP_AEC_CollabBP/Config/SteamVRBindings/vive_tracker.json
deleted file mode 100644
diff --git a/Templates/TP_AEC_CollabBP/Config/SteamVRBindings/vive_tracker_camera.json b/Templates/TP_AEC_CollabBP/Config/SteamVRBindings/vive_tracker_camera.json
deleted file mode 100644
diff --git a/Templates/TP_AEC_CollabBP/Config/SteamVRBindings/vive_tracker_chest.json b/Templates/TP_AEC_CollabBP/Config/SteamVRBindings/vive_tracker_chest.json
deleted file mode 100644
diff --git a/Templates/TP_AEC_CollabBP/Config/SteamVRBindings/vive_tracker_handed.json b/Templates/TP_AEC_CollabBP/Config/SteamVRBindings/vive_tracker_handed.json
deleted file mode 100644
diff --git a/Templates/TP_AEC_CollabBP/Config/SteamVRBindings/vive_tracker_keyboard.json b/Templates/TP_AEC_CollabBP/Config/SteamVRBindings/vive_tracker_keyboard.json
deleted file mode 100644
diff --git a/Templates/TP_AEC_CollabBP/Config/SteamVRBindings/vive_tracker_left_elbow.json b/Templates/TP_AEC_CollabBP/Config/SteamVRBindings/vive_tracker_left_elbow.json
deleted file mode 100644
diff --git a/Templates/TP_AEC_CollabBP/Config/SteamVRBindings/vive_tracker_left_foot.json b/Templates/TP_AEC_CollabBP/Config/SteamVRBindings/vive_tracker_left_foot.json
deleted file mode 100644
diff --git a/Templates/TP_AEC_CollabBP/Config/SteamVRBindings/vive_tracker_left_knee.json b/Templates/TP_AEC_CollabBP/Config/SteamVRBindings/vive_tracker_left_knee.json
deleted file mode 100644
diff --git a/Templates/TP_AEC_CollabBP/Config/SteamVRBindings/vive_tracker_left_shoulder.json b/Templates/TP_AEC_CollabBP/Config/SteamVRBindings/vive_tracker_left_shoulder.json
deleted file mode 100644
diff --git a/Templates/TP_AEC_CollabBP/Config/SteamVRBindings/vive_tracker_right_elbow.json b/Templates/TP_AEC_CollabBP/Config/SteamVRBindings/vive_tracker_right_elbow.json
deleted file mode 100644
diff --git a/Templates/TP_AEC_CollabBP/Config/SteamVRBindings/vive_tracker_right_foot.json b/Templates/TP_AEC_CollabBP/Config/SteamVRBindings/vive_tracker_right_foot.json
deleted file mode 100644
diff --git a/Templates/TP_AEC_CollabBP/Config/SteamVRBindings/vive_tracker_right_knee.json b/Templates/TP_AEC_CollabBP/Config/SteamVRBindings/vive_tracker_right_knee.json
deleted file mode 100644
diff --git a/Templates/TP_AEC_CollabBP/Config/SteamVRBindings/vive_tracker_right_shoulder.json b/Templates/TP_AEC_CollabBP/Config/SteamVRBindings/vive_tracker_right_shoulder.json
deleted file mode 100644
diff --git a/Templates/TP_AEC_CollabBP/Config/SteamVRBindings/vive_tracker_waist.json b/Templates/TP_AEC_CollabBP/Config/SteamVRBindings/vive_tracker_waist.json
deleted file mode 100644
diff --git a/Templates/TP_AEC_CollabBP/Config/steamvr_ue_editor_app.json b/Templates/TP_AEC_CollabBP/Config/steamvr_ue_editor_app.json
deleted file mode 100644
diff --git a/Templates/TP_MFG_CollabBP/Config/DefaultEngine.ini b/Templates/TP_MFG_CollabBP/Config/DefaultEngine.ini
--- a/Templates/TP_MFG_CollabBP/Config/DefaultEngine.ini
+++ b/Templates/TP_MFG_CollabBP/Config/DefaultEngine.ini
@@ -130,9 +130,6 @@ ChaosSettings=(DefaultThreadingModel=TaskGraph,DedicatedThreadTickMode=VariableC
 IconModelPath=(Path="")
 IconPortalPath=(Path="")
 
-[/Script/OculusHMD.OculusHMDRuntimeSettings]
-bAutoEnabled=False
-
 [/Script/Engine.AudioSettings]
 MasterSubmix=/Engine/EngineSounds/Submixes/MasterSubmixDefault.MasterSubmixDefault
 
diff --git a/Templates/TP_MFG_CollabBP/Config/DefaultInput.ini b/Templates/TP_MFG_CollabBP/Config/DefaultInput.ini
--- a/Templates/TP_MFG_CollabBP/Config/DefaultInput.ini
+++ b/Templates/TP_MFG_CollabBP/Config/DefaultInput.ini
@@ -137,6 +137,12 @@ bUseMouseForTouch=False
 bEnableMouseSmoothing=True
 bEnableFOVScaling=True
 bCaptureMouseOnLaunch=True
+bEnableLegacyInputScales=True
+bEnableMotionControls=True
+bFilterInputByPlatformUser=False
+bEnableInputDeviceSubsystem=True
+bShouldFlushPressedKeysOnViewportFocusLost=True
+bEnableDynamicComponentInputBinding=True
 bAlwaysShowTouchInterface=False
 bShowConsoleOnFourFingerTap=True
 bEnableGestureRecognizer=True
@@ -145,95 +151,12 @@ DefaultViewportMouseCaptureMode=CapturePermanently_IncludingInitialMouseDown
 DefaultViewportMouseLockMode=LockOnCapture
 FOVScale=0.011110
 DoubleClickTime=0.200000
-+ActionMappings=(ActionName="Bookmark_1",bShift=False,bCtrl=False,bAlt=False,bCmd=False,Key=One)
-+ActionMappings=(ActionName="Bookmark_2",bShift=False,bCtrl=False,bAlt=False,bCmd=False,Key=Two)
-+ActionMappings=(ActionName="Bookmark_3",bShift=False,bCtrl=False,bAlt=False,bCmd=False,Key=Three)
-+ActionMappings=(ActionName="Bookmark_4",bShift=False,bCtrl=False,bAlt=False,bCmd=False,Key=Four)
-+ActionMappings=(ActionName="Bookmark_5",bShift=False,bCtrl=False,bAlt=False,bCmd=False,Key=Five)
-+ActionMappings=(ActionName="Bookmark_6",bShift=False,bCtrl=False,bAlt=False,bCmd=False,Key=Six)
-+ActionMappings=(ActionName="Bookmark_7",bShift=False,bCtrl=False,bAlt=False,bCmd=False,Key=Seven)
-+ActionMappings=(ActionName="Bookmark_8",bShift=False,bCtrl=False,bAlt=False,bCmd=False,Key=Eight)
-+ActionMappings=(ActionName="Bookmark_9",bShift=False,bCtrl=False,bAlt=False,bCmd=False,Key=Nine)
-+ActionMappings=(ActionName="FlyMode",bShift=False,bCtrl=False,bAlt=False,bCmd=False,Key=U)
-+ActionMappings=(ActionName="LockToAxis",bShift=False,bCtrl=False,bAlt=False,bCmd=False,Key=LeftShift)
-+ActionMappings=(ActionName="OpenContextualMenu",bShift=False,bCtrl=False,bAlt=False,bCmd=False,Key=SpaceBar)
-+ActionMappings=(ActionName="OrbitAroundObject",bShift=False,bCtrl=False,bAlt=False,bCmd=False,Key=RightMouseButton)
-+ActionMappings=(ActionName="OrbitMode",bShift=False,bCtrl=False,bAlt=False,bCmd=False,Key=O)
-+ActionMappings=(ActionName="OrbitZoomIn",bShift=False,bCtrl=False,bAlt=False,bCmd=False,Key=MouseScrollUp)
-+ActionMappings=(ActionName="OrbitZoomOut",bShift=False,bCtrl=False,bAlt=False,bCmd=False,Key=MouseScrollDown)
-+ActionMappings=(ActionName="PushAndPull",bShift=False,bCtrl=False,bAlt=False,bCmd=False,Key=MouseWheelAxis)
-+ActionMappings=(ActionName="Quit",bShift=False,bCtrl=True,bAlt=False,bCmd=False,Key=Q)
-+ActionMappings=(ActionName="RotateView",bShift=False,bCtrl=False,bAlt=False,bCmd=False,Key=RightMouseButton)
-+ActionMappings=(ActionName="SetPointOfinterest",bShift=False,bCtrl=False,bAlt=False,bCmd=False,Key=MiddleMouseButton)
-+ActionMappings=(ActionName="ShootLaser",bShift=False,bCtrl=False,bAlt=False,bCmd=False,Key=LeftMouseButton)
-+ActionMappings=(ActionName="ShowContextualMenu_VR_L",bShift=False,bCtrl=False,bAlt=False,bCmd=False,Key=OculusTouch_Left_X_Click)
-+ActionMappings=(ActionName="ShowContextualMenu_VR_L",bShift=False,bCtrl=False,bAlt=False,bCmd=False,Key=ValveIndex_Left_A_Click)
-+ActionMappings=(ActionName="ShowContextualMenu_VR_L",bShift=False,bCtrl=False,bAlt=False,bCmd=False,Key=Vive_Left_Menu_Click)
-+ActionMappings=(ActionName="ShowContextualMenu_VR_R",bShift=False,bCtrl=False,bAlt=False,bCmd=False,Key=OculusTouch_Right_A_Click)
-+ActionMappings=(ActionName="ShowContextualMenu_VR_R",bShift=False,bCtrl=False,bAlt=False,bCmd=False,Key=ValveIndex_Right_A_Click)
-+ActionMappings=(ActionName="ShowContextualMenu_VR_R",bShift=False,bCtrl=False,bAlt=False,bCmd=False,Key=Vive_Right_Menu_Click)
-+ActionMappings=(ActionName="VR_Teleport_L",bShift=False,bCtrl=False,bAlt=False,bCmd=False,Key=OculusTouch_Left_Y_Click)
-+ActionMappings=(ActionName="VR_Teleport_L",bShift=False,bCtrl=False,bAlt=False,bCmd=False,Key=ValveIndex_Left_B_Click)
-+ActionMappings=(ActionName="VR_Teleport_L",bShift=False,bCtrl=False,bAlt=False,bCmd=False,Key=Vive_Left_Grip_Click)
-+ActionMappings=(ActionName="VR_Teleport_R",bShift=False,bCtrl=False,bAlt=False,bCmd=False,Key=OculusTouch_Right_B_Click)
-+ActionMappings=(ActionName="VR_Teleport_R",bShift=False,bCtrl=False,bAlt=False,bCmd=False,Key=ValveIndex_Right_B_Click)
-+ActionMappings=(ActionName="VR_Teleport_R",bShift=False,bCtrl=False,bAlt=False,bCmd=False,Key=Vive_Right_Grip_Click)
-+ActionMappings=(ActionName="VR_Trigger_L",bShift=False,bCtrl=False,bAlt=False,bCmd=False,Key=OculusTouch_Left_Trigger_Click)
-+ActionMappings=(ActionName="VR_Trigger_L",bShift=False,bCtrl=False,bAlt=False,bCmd=False,Key=ValveIndex_Left_Trigger_Click)
-+ActionMappings=(ActionName="VR_Trigger_L",bShift=False,bCtrl=False,bAlt=False,bCmd=False,Key=Vive_Left_Trigger_Click)
-+ActionMappings=(ActionName="VR_Trigger_R",bShift=False,bCtrl=False,bAlt=False,bCmd=False,Key=OculusTouch_Right_Trigger_Click)
-+ActionMappings=(ActionName="VR_Trigger_R",bShift=False,bCtrl=False,bAlt=False,bCmd=False,Key=ValveIndex_Right_Trigger_Click)
-+ActionMappings=(ActionName="VR_Trigger_R",bShift=False,bCtrl=False,bAlt=False,bCmd=False,Key=Vive_Right_Trigger_Click)
-+ActionMappings=(ActionName="VRMode",bShift=False,bCtrl=False,bAlt=False,bCmd=False,Key=P)
-+ActionMappings=(ActionName="WalkMode",bShift=False,bCtrl=False,bAlt=False,bCmd=False,Key=I)
-+ActionMappings=(ActionName="NoCommand",bShift=False,bCtrl=False,bAlt=False,bCmd=False,Key=Escape)
-+ActionMappings=(ActionName="TurnRight",bShift=False,bCtrl=False,bAlt=False,bCmd=False,Key=D)
-+ActionMappings=(ActionName="TurnLeft",bShift=False,bCtrl=False,bAlt=False,bCmd=False,Key=A)
-+ActionMappings=(ActionName="TurnRight",bShift=False,bCtrl=False,bAlt=False,bCmd=False,Key=OculusTouch_Right_Thumbstick_Right)
-+ActionMappings=(ActionName="TurnLeft",bShift=False,bCtrl=False,bAlt=False,bCmd=False,Key=OculusTouch_Right_Thumbstick_Left)
-+ActionMappings=(ActionName="VR_Teleport_L",bShift=False,bCtrl=False,bAlt=False,bCmd=False,Key=MixedReality_Left_Grip_Click)
-+ActionMappings=(ActionName="VR_Teleport_R",bShift=False,bCtrl=False,bAlt=False,bCmd=False,Key=MixedReality_Right_Grip_Click)
-+ActionMappings=(ActionName="VR_Trigger_R",bShift=False,bCtrl=False,bAlt=False,bCmd=False,Key=MixedReality_Right_Trigger_Click)
-+ActionMappings=(ActionName="VR_Trigger_L",bShift=False,bCtrl=False,bAlt=False,bCmd=False,Key=MixedReality_Left_Trigger_Click)
-+ActionMappings=(ActionName="ShowContextualMenu_VR_R",bShift=False,bCtrl=False,bAlt=False,bCmd=False,Key=MixedReality_Right_Menu_Click)
-+ActionMappings=(ActionName="ShowContextualMenu_VR_L",bShift=False,bCtrl=False,bAlt=False,bCmd=False,Key=MixedReality_Left_Menu_Click)
-+AxisMappings=(AxisName="FlyUp",Scale=1.000000,Key=E)
-+AxisMappings=(AxisName="FlyUp",Scale=-1.000000,Key=Q)
-+AxisMappings=(AxisName="LookRight",Scale=1.000000,Key=Gamepad_RightX)
-+AxisMappings=(AxisName="LookRight",Scale=1.000000,Key=MouseX)
-+AxisMappings=(AxisName="LookUp",Scale=1.000000,Key=Gamepad_RightY)
-+AxisMappings=(AxisName="LookUp",Scale=1.000000,Key=MouseY)
-+AxisMappings=(AxisName="MoveForward",Scale=-1.000000,Key=Down)
-+AxisMappings=(AxisName="MoveForward",Scale=1.000000,Key=Gamepad_LeftY)
-+AxisMappings=(AxisName="MoveForward",Scale=-1.000000,Key=S)
-+AxisMappings=(AxisName="MoveForward",Scale=1.000000,Key=Up)
-+AxisMappings=(AxisName="MoveForward",Scale=1.000000,Key=W)
-+AxisMappings=(AxisName="MoveRight",Scale=-1.000000,Key=A)
-+AxisMappings=(AxisName="MoveRight",Scale=1.000000,Key=D)
-+AxisMappings=(AxisName="MoveRight",Scale=1.000000,Key=Gamepad_LeftX)
-+AxisMappings=(AxisName="pinch",Scale=1.000000,Key=Gesture_Pinch)
-+AxisMappings=(AxisName="Thumbstick_L_Y",Scale=-1.000000,Key=OculusTouch_Left_Thumbstick_Y)
-+AxisMappings=(AxisName="Thumbstick_L_Y",Scale=-1.000000,Key=ValveIndex_Left_Thumbstick_Y)
-+AxisMappings=(AxisName="Thumbstick_R_Y",Scale=-1.000000,Key=OculusTouch_Right_Thumbstick_Y)
-+AxisMappings=(AxisName="Thumbstick_R_Y",Scale=-1.000000,Key=ValveIndex_Right_Thumbstick_Y)
-+AxisMappings=(AxisName="ThumbStick_Y_Axis",Scale=1.000000,Key=MouseWheelAxis)
-+AxisMappings=(AxisName="ThumbStick_Y_Axis",Scale=1.000000,Key=OculusTouch_Left_Thumbstick_Y)
-+AxisMappings=(AxisName="ThumbStick_Y_Axis",Scale=1.000000,Key=OculusTouch_Right_Thumbstick_Y)
-+AxisMappings=(AxisName="WalkForward",Scale=-1.000000,Key=Down)
-+AxisMappings=(AxisName="WalkForward",Scale=1.000000,Key=Gamepad_LeftY)
-+AxisMappings=(AxisName="WalkForward",Scale=-1.000000,Key=S)
-+AxisMappings=(AxisName="WalkForward",Scale=1.000000,Key=Up)
-+AxisMappings=(AxisName="WalkForward",Scale=1.000000,Key=W)
-+AxisMappings=(AxisName="WalkRight",Scale=-1.000000,Key=A)
-+AxisMappings=(AxisName="WalkRight",Scale=1.000000,Key=D)
-+AxisMappings=(AxisName="WalkRight",Scale=1.000000,Key=Gamepad_LeftX)
-+AxisMappings=(AxisName="WalkRight",Scale=-1.000000,Key=Left)
-+AxisMappings=(AxisName="WalkRight",Scale=1.000000,Key=Right)
-+AxisMappings=(AxisName="Thumbstick_L_Y",Scale=-1.000000,Key=Vive_Left_Trackpad_Y)
-+AxisMappings=(AxisName="Thumbstick_R_Y",Scale=-1.000000,Key=Vive_Right_Trackpad_Y)
-DefaultPlayerInputClass=/Script/Engine.PlayerInput
-DefaultInputComponentClass=/Script/Engine.InputComponent
+DefaultPlayerInputClass=/Script/EnhancedInput.EnhancedPlayerInput
+DefaultInputComponentClass=/Script/EnhancedInput.EnhancedInputComponent
 DefaultTouchInterface=/Game/CollaborativeViewer/Blueprints/Pawn/Touch/TouchInterface.TouchInterface
 -ConsoleKeys=Tilde
 +ConsoleKeys=Tilde
 
+[/Script/OpenXRInput.OpenXRInputSettings]
+MappableInputConfig=/Game/CollaborativeViewer/Input/PMI_CollabTemplate.PMI_CollabTemplate
+
diff --git a/Templates/TP_MFG_CollabBP/Config/SteamVRBindings/gamepad.json b/Templates/TP_MFG_CollabBP/Config/SteamVRBindings/gamepad.json
deleted file mode 100644
diff --git a/Templates/TP_MFG_CollabBP/Config/SteamVRBindings/holographic_controller.json b/Templates/TP_MFG_CollabBP/Config/SteamVRBindings/holographic_controller.json
deleted file mode 100644
diff --git a/Templates/TP_MFG_CollabBP/Config/SteamVRBindings/hpmotioncontroller.json b/Templates/TP_MFG_CollabBP/Config/SteamVRBindings/hpmotioncontroller.json
deleted file mode 100644
diff --git a/Templates/TP_MFG_CollabBP/Config/SteamVRBindings/indexhmd.json b/Templates/TP_MFG_CollabBP/Config/SteamVRBindings/indexhmd.json
deleted file mode 100644
diff --git a/Templates/TP_MFG_CollabBP/Config/SteamVRBindings/knuckles.json b/Templates/TP_MFG_CollabBP/Config/SteamVRBindings/knuckles.json
deleted file mode 100644
diff --git a/Templates/TP_MFG_CollabBP/Config/SteamVRBindings/oculus_touch.json b/Templates/TP_MFG_CollabBP/Config/SteamVRBindings/oculus_touch.json
deleted file mode 100644
diff --git a/Templates/TP_MFG_CollabBP/Config/SteamVRBindings/rift.json b/Templates/TP_MFG_CollabBP/Config/SteamVRBindings/rift.json
deleted file mode 100644
diff --git a/Templates/TP_MFG_CollabBP/Config/SteamVRBindings/steamvr_manifest.json b/Templates/TP_MFG_CollabBP/Config/SteamVRBindings/steamvr_manifest.json
deleted file mode 100644
diff --git a/Templates/TP_MFG_CollabBP/Config/SteamVRBindings/vive.json b/Templates/TP_MFG_CollabBP/Config/SteamVRBindings/vive.json
deleted file mode 100644
diff --git a/Templates/TP_MFG_CollabBP/Config/SteamVRBindings/vive_controller.json b/Templates/TP_MFG_CollabBP/Config/SteamVRBindings/vive_controller.json
deleted file mode 100644
diff --git a/Templates/TP_MFG_CollabBP/Config/SteamVRBindings/vive_cosmos_controller.json b/Templates/TP_MFG_CollabBP/Config/SteamVRBindings/vive_cosmos_controller.json
deleted file mode 100644
diff --git a/Templates/TP_MFG_CollabBP/Config/SteamVRBindings/vive_pro.json b/Templates/TP_MFG_CollabBP/Config/SteamVRBindings/vive_pro.json
deleted file mode 100644
diff --git a/Templates/TP_MFG_CollabBP/Config/SteamVRBindings/vive_tracker.json b/Templates/TP_MFG_CollabBP/Config/SteamVRBindings/vive_tracker.json
deleted file mode 100644
diff --git a/Templates/TP_MFG_CollabBP/Config/SteamVRBindings/vive_tracker_camera.json b/Templates/TP_MFG_CollabBP/Config/SteamVRBindings/vive_tracker_camera.json
deleted file mode 100644
diff --git a/Templates/TP_MFG_CollabBP/Config/SteamVRBindings/vive_tracker_chest.json b/Templates/TP_MFG_CollabBP/Config/SteamVRBindings/vive_tracker_chest.json
deleted file mode 100644
diff --git a/Templates/TP_MFG_CollabBP/Config/SteamVRBindings/vive_tracker_handed.json b/Templates/TP_MFG_CollabBP/Config/SteamVRBindings/vive_tracker_handed.json
deleted file mode 100644
diff --git a/Templates/TP_MFG_CollabBP/Config/SteamVRBindings/vive_tracker_keyboard.json b/Templates/TP_MFG_CollabBP/Config/SteamVRBindings/vive_tracker_keyboard.json
deleted file mode 100644
diff --git a/Templates/TP_MFG_CollabBP/Config/SteamVRBindings/vive_tracker_left_elbow.json b/Templates/TP_MFG_CollabBP/Config/SteamVRBindings/vive_tracker_left_elbow.json
deleted file mode 100644
diff --git a/Templates/TP_MFG_CollabBP/Config/SteamVRBindings/vive_tracker_left_foot.json b/Templates/TP_MFG_CollabBP/Config/SteamVRBindings/vive_tracker_left_foot.json
deleted file mode 100644
diff --git a/Templates/TP_MFG_CollabBP/Config/SteamVRBindings/vive_tracker_left_knee.json b/Templates/TP_MFG_CollabBP/Config/SteamVRBindings/vive_tracker_left_knee.json
deleted file mode 100644
diff --git a/Templates/TP_MFG_CollabBP/Config/SteamVRBindings/vive_tracker_left_shoulder.json b/Templates/TP_MFG_CollabBP/Config/SteamVRBindings/vive_tracker_left_shoulder.json
deleted file mode 100644
diff --git a/Templates/TP_MFG_CollabBP/Config/SteamVRBindings/vive_tracker_right_elbow.json b/Templates/TP_MFG_CollabBP/Config/SteamVRBindings/vive_tracker_right_elbow.json
deleted file mode 100644
diff --git a/Templates/TP_MFG_CollabBP/Config/SteamVRBindings/vive_tracker_right_foot.json b/Templates/TP_MFG_CollabBP/Config/SteamVRBindings/vive_tracker_right_foot.json
deleted file mode 100644
diff --git a/Templates/TP_MFG_CollabBP/Config/SteamVRBindings/vive_tracker_right_knee.json b/Templates/TP_MFG_CollabBP/Config/SteamVRBindings/vive_tracker_right_knee.json
deleted file mode 100644
diff --git a/Templates/TP_MFG_CollabBP/Config/SteamVRBindings/vive_tracker_right_shoulder.json b/Templates/TP_MFG_CollabBP/Config/SteamVRBindings/vive_tracker_right_shoulder.json
deleted file mode 100644
diff --git a/Templates/TP_MFG_CollabBP/Config/SteamVRBindings/vive_tracker_waist.json b/Templates/TP_MFG_CollabBP/Config/SteamVRBindings/vive_tracker_waist.json
deleted file mode 100644
diff --git a/Templates/TP_MFG_CollabBP/Config/steamvr_ue_editor_app.json b/Templates/TP_MFG_CollabBP/Config/steamvr_ue_editor_app.json
deleted file mode 100644
diff --git a/Templates/TP_MFG_CollabBP/TP_MFG_CollabBP.uproject b/Templates/TP_MFG_CollabBP/TP_MFG_CollabBP.uproject
--- a/Templates/TP_MFG_CollabBP/TP_MFG_CollabBP.uproject
+++ b/Templates/TP_MFG_CollabBP/TP_MFG_CollabBP.uproject
@@ -36,7 +36,8 @@
 			"Name": "DatasmithCADImporter",
 			"Enabled": true,
 			"SupportedTargetPlatforms": [
-				"Win64"
+				"Win64",
+				"Linux"
 			]
 		},
 		{
@@ -55,42 +56,13 @@
 			"Enabled": true
 		},
 		{
-			"Name": "OculusVR",
-			"Enabled": false,
-			"SupportedTargetPlatforms": [
-				"Win32",
-				"Win64",
-				"Android"
-			]
-		},
-		{
-			"Name": "OpenXREyeTracker",
-			"Enabled": false,
-			"SupportedTargetPlatforms": [
-				"Win32",
-				"Win64",
-				"Linux",
-				"HoloLens",
-				"Android"
-			]
-		},
-		{
-			"Name": "OpenXRHandTracking",
-			"Enabled": false,
+			"Name": "OpenXR",
+			"Enabled": true,
 			"SupportedTargetPlatforms": [
-				"Win32",
 				"Win64",
 				"Linux",
-				"HoloLens",
-				"Android"
-			]
-		},
-		{
-			"Name": "SteamVR",
-			"Enabled": false,
-			"SupportedTargetPlatforms": [
-				"Win32",
-				"Win64"
+				"Android",
+				"HoloLens"
 			]
 		}
 	]
